Disable the blank plugin by default. It's not required by anything.

diff --git a/Engine/Plugins/Developer/BlankPlugin/BlankPlugin.uplugin b/Engine/Plugins/Developer/BlankPlugin/BlankPlugin.uplugin
--- a/Engine/Plugins/Developer/BlankPlugin/BlankPlugin.uplugin
+++ b/Engine/Plugins/Developer/BlankPlugin/BlankPlugin.uplugin
@@ -9,7 +9,6 @@
 	"EngineVersion" : "4.2.0",
 	"Description" : "An example of a minimal plugin.  This can be used as a starting point when creating your own plugin.",
 	"Category" : "Examples",
-	"EnabledByDefault" : true,
 	
 	"Modules" :
 	[
